Marked ZoneGraph plugins as experimental

diff --git a/Engine/Plugins/Runtime/ZoneGraph/ZoneGraph.uplugin b/Engine/Plugins/Runtime/ZoneGraph/ZoneGraph.uplugin
--- a/Engine/Plugins/Runtime/ZoneGraph/ZoneGraph.uplugin
+++ b/Engine/Plugins/Runtime/ZoneGraph/ZoneGraph.uplugin
@@ -1,7 +1,7 @@
 {
 	"FileVersion" : 3,
 	"Version" : 1,
-	"VersionName" : "1.0",
+	"VersionName" : "0.5",
 	"FriendlyName" : "ZoneGraph",
 	"Description" : "Description missing.",
 	"Category" : "AI",
@@ -12,7 +12,8 @@
 	"SupportURL" : "",
 	"EnabledByDefault" : false,
 	"CanContainContent" : false,
-	"IsBetaVersion" : false,
+	"IsBetaVersion": false,
+	"IsExperimentalVersion": true,
 	"Installed" : false,
 	"Modules" :
 	[
diff --git a/Engine/Plugins/Runtime/ZoneGraphAnnotations/ZoneGraphAnnotations.uplugin b/Engine/Plugins/Runtime/ZoneGraphAnnotations/ZoneGraphAnnotations.uplugin
--- a/Engine/Plugins/Runtime/ZoneGraphAnnotations/ZoneGraphAnnotations.uplugin
+++ b/Engine/Plugins/Runtime/ZoneGraphAnnotations/ZoneGraphAnnotations.uplugin
@@ -1,7 +1,7 @@
 {
 	"FileVersion": 1,
 	"Version": 1,
-	"VersionName": "1.0",
+	"VersionName": "0.5",
 	"FriendlyName": "ZoneGraph Annotations",
 	"Description": "Annotations for Zone Graph",
 	"Category": "AI",
@@ -13,7 +13,7 @@
 	"EnabledByDefault" : false,
 	"CanContainContent": true,
 	"IsBetaVersion": false,
-	"IsExperimentalVersion": false,
+	"IsExperimentalVersion": true,
 	"Installed": false,
 	"Modules": [
 		{
