Fixed ZoneGraphDisturbanceAnnotation find predicate

diff --git a/Engine/Plugins/Runtime/ZoneGraphAnnotations/Source/ZoneGraphAnnotations/Private/Annotations/ZoneGraphDisturbanceAnnotation.cpp b/Engine/Plugins/Runtime/ZoneGraphAnnotations/Source/ZoneGraphAnnotations/Private/Annotations/ZoneGraphDisturbanceAnnotation.cpp
--- a/Engine/Plugins/Runtime/ZoneGraphAnnotations/Source/ZoneGraphAnnotations/Private/Annotations/ZoneGraphDisturbanceAnnotation.cpp
+++ b/Engine/Plugins/Runtime/ZoneGraphAnnotations/Source/ZoneGraphAnnotations/Private/Annotations/ZoneGraphDisturbanceAnnotation.cpp
@@ -516,7 +516,7 @@ void UZoneGraphDisturbanceAnnotation::HandleEvents(TConstArrayView<const UScript
 	{
 		if (Obstacle.Action == EZoneGraphObstacleDisturbanceAreaAction::Add)
 		{
-			FZoneGraphObstacleDisturbanceArea* ExistingObstacle = Obstacles.FindByPredicate([Obstacle](const FZoneGraphObstacleDisturbanceArea& Area) { return Obstacle.ObstacleID == Obstacle.ObstacleID; });
+			FZoneGraphObstacleDisturbanceArea* ExistingObstacle = Obstacles.FindByPredicate([Obstacle](const FZoneGraphObstacleDisturbanceArea& Area) { return Area.ObstacleID == Obstacle.ObstacleID; });
 			if (ExistingObstacle)
 			{
 				*ExistingObstacle = Obstacle;
