Fix Mass Crowd flee behavior and crash in editor

- Made AMassSpawner visible in non-shipping/testing so that component debug draw works
- Fixed crash in UZoneGraphAnnotationComponent::OnUnregister() when component is edited BP editor
- Added debug draw culling to ZoneGraph Disturbance annotations and ZoneGraph Crowd Lane annotations
- Fixed flee lane tags on CitySample (flee would try to use sidewalks, which causes problems)
- Added fallback for flee calculation when there are no exists
- Added helper functions to UZoneGraphAnnotationComponent to add distance culling for debug draws

diff --git a/Engine/Plugins/AI/MassCrowd/Source/MassCrowd/Private/ZoneGraphCrowdLaneAnnotations.cpp b/Engine/Plugins/AI/MassCrowd/Source/MassCrowd/Private/ZoneGraphCrowdLaneAnnotations.cpp
--- a/Engine/Plugins/AI/MassCrowd/Source/MassCrowd/Private/ZoneGraphCrowdLaneAnnotations.cpp
+++ b/Engine/Plugins/AI/MassCrowd/Source/MassCrowd/Private/ZoneGraphCrowdLaneAnnotations.cpp
@@ -85,7 +85,10 @@ void UZoneGraphCrowdLaneAnnotations::TickAnnotation(const float DeltaTime, FZone
 	StateChangeEvents.Reset();
 
 #if !UE_BUILD_SHIPPING && !UE_BUILD_TEST
-	MarkRenderStateDirty();
+	if (bEnableDebugDrawing)
+	{
+		MarkRenderStateDirty();
+	}
 #endif
 }
 
@@ -110,6 +113,13 @@ void UZoneGraphCrowdLaneAnnotations::DebugDraw(FZoneGraphAnnotationSceneProxy* D
 	AllTags.Add(CloseLaneTag);
 	AllTags.Add(WaitingLaneTag);
 
+	FVector ViewLocation = FVector::ZeroVector;
+	FRotator ViewRotation = FRotator::ZeroRotator;
+	GetFirstViewPoint(ViewLocation, ViewRotation);
+
+	const float DrawDistance = GetMaxDebugDrawDistance();
+	const float DrawDistanceSq = FMath::Square(DrawDistance);
+
 	for (const FRegisteredCrowdLaneData& RegisteredLaneData : CrowdSubsystem->RegisteredLaneData)
 	{
 		const FZoneGraphStorage* ZoneStorage = ZoneGraph->GetZoneGraphStorage(RegisteredLaneData.DataHandle);
@@ -118,25 +128,34 @@ void UZoneGraphCrowdLaneAnnotations::DebugDraw(FZoneGraphAnnotationSceneProxy* D
 			continue;
 		}
 
-		for (int32 LaneIndex = 0; LaneIndex < RegisteredLaneData.CrowdLaneDataArray.Num(); LaneIndex++)
+		for (const FZoneData& Zone : ZoneStorage->Zones)
 		{
-			const FZoneGraphCrowdLaneData& LaneData = RegisteredLaneData.CrowdLaneDataArray[LaneIndex];
-			if (LaneData.GetState() == ECrowdLaneState::Closed)
+			const float DistanceSq = FVector::DistSquared(ViewLocation, Zone.Bounds.GetCenter());
+			if (DistanceSq > DrawDistanceSq)
 			{
-				const FZoneGraphLaneHandle LaneHandle(LaneIndex, RegisteredLaneData.DataHandle);
-
-				FLinearColor Color = ClosedColor;
+				continue;
+			}
 
-				const FCrowdWaitAreaData* WaitArea = CrowdSubsystem->GetCrowdWaitingAreaData(LaneHandle);
-				if (WaitArea && !WaitArea->IsFull())
+			for (int32 LaneIdx = Zone.LanesBegin; LaneIdx < Zone.LanesEnd; LaneIdx++)
+			{
+				const FZoneGraphCrowdLaneData& LaneData = RegisteredLaneData.CrowdLaneDataArray[LaneIdx];
+				if (LaneData.GetState() == ECrowdLaneState::Closed)
 				{
-					Color = WaitingColor;
-				}
+					const FZoneGraphLaneHandle LaneHandle(LaneIdx, RegisteredLaneData.DataHandle);
 
-				UE::ZoneGraph::RenderingUtilities::AppendLane(DebugProxy, *ZoneStorage, LaneHandle, Color.ToFColor(/*sRGB*/true), 4.0f, ZOffset);
+					FLinearColor Color = ClosedColor;
+
+					const FCrowdWaitAreaData* WaitArea = CrowdSubsystem->GetCrowdWaitingAreaData(LaneHandle);
+					if (WaitArea && !WaitArea->IsFull())
+					{
+						Color = WaitingColor;
+					}
+
+					UE::ZoneGraph::RenderingUtilities::AppendLane(DebugProxy, *ZoneStorage, LaneHandle, Color.ToFColor(/*sRGB*/true), 4.0f, ZOffset);
+				}
 			}
 		}
-
+		
 		auto AppendCircleXY = [DebugProxy](const FVector& Center, const float Radius, const FColor Color, const float LineThickness)
 		{
 			static int32 NumDivs = 16;
@@ -160,10 +179,17 @@ void UZoneGraphCrowdLaneAnnotations::DebugDraw(FZoneGraphAnnotationSceneProxy* D
 		const FColor SlotColor = FColor::Orange;
 		for (const FCrowdWaitAreaData& WaitArea : RegisteredLaneData.WaitAreas)
 		{
-			for (const FCrowdWaitSlot& Slot : WaitArea.Slots)
+			if (WaitArea.Slots.Num() > 0)
 			{
-				AppendCircleXY(Slot.Position + ZOffset, Slot.Radius, SlotColor, 1.0f);
-				DebugProxy->Lines.Emplace(Slot.Position + ZOffset, Slot.Position + Slot.Forward * Slot.Radius + ZOffset, SlotColor, 4.0f);
+				const float DistanceSq = FVector::DistSquared(ViewLocation, WaitArea.Slots[0].Position);
+				if (DistanceSq < DrawDistanceSq)
+				{
+					for (const FCrowdWaitSlot& Slot : WaitArea.Slots)
+					{
+						AppendCircleXY(Slot.Position + ZOffset, Slot.Radius, SlotColor, 1.0f);
+						DebugProxy->Lines.Emplace(Slot.Position + ZOffset, Slot.Position + Slot.Forward * Slot.Radius + ZOffset, SlotColor, 4.0f);
+					}
+				}
 			}
 		}
 	}
@@ -179,6 +205,8 @@ void UZoneGraphCrowdLaneAnnotations::DebugDrawCanvas(UCanvas* Canvas, APlayerCon
 	const FColor OldDrawColor = Canvas->DrawColor;
 	const UFont* RenderFont = GEngine->GetSmallFont();
 
+	const FFontRenderInfo FontInfo = Canvas->CreateFontRenderInfo(true, true);
+
 	Canvas->SetDrawColor(FColor::White);
 	static const FVector ZOffset(0, 0, 35.0f);
 
@@ -189,6 +217,20 @@ void UZoneGraphCrowdLaneAnnotations::DebugDrawCanvas(UCanvas* Canvas, APlayerCon
 		return;
 	}
 
+	if (Canvas->SceneView == nullptr)
+	{
+		return;
+	}
+	
+	const FVector ViewLocation = Canvas->SceneView->ViewLocation;
+	const float DrawDistance = GetMaxDebugDrawDistance() * 0.25f;
+	const float DrawDistanceSq = FMath::Square(DrawDistance);
+
+	auto InFrustum = [Canvas](const FVector& Location)
+	{
+		return Canvas->SceneView->ViewFrustum.IntersectBox(Location, FVector::ZeroVector);
+	};
+
 	for (const FRegisteredCrowdLaneData& RegisteredLaneData : CrowdSubsystem->RegisteredLaneData)
 	{
 		const FZoneGraphStorage* ZoneStorage = RegisteredLaneData.DataHandle.IsValid() ? ZoneGraph->GetZoneGraphStorage(RegisteredLaneData.DataHandle) : nullptr;
@@ -197,37 +239,67 @@ void UZoneGraphCrowdLaneAnnotations::DebugDrawCanvas(UCanvas* Canvas, APlayerCon
 			continue;
 		}
 
-		for (int32 LaneIdx = 0; LaneIdx < ZoneStorage->Lanes.Num() && ZoneGraphAnnotationSubsystem; LaneIdx++)
+		for (const FZoneData& Zone : ZoneStorage->Zones)
 		{
-			FZoneGraphLaneLocation CenterLoc;
-			UE::ZoneGraph::Query::CalculateLocationAlongLaneFromRatio(*ZoneStorage, LaneIdx, 0.5f, CenterLoc);
-			const FVector ScreenLoc = Canvas->Project(CenterLoc.Position);
-
-			const FZoneGraphTagMask Mask = ZoneGraphAnnotationSubsystem->GetAnnotationTags({LaneIdx, RegisteredLaneData.DataHandle});
-			Canvas->DrawText(RenderFont, FString::Printf(TEXT("%s\n0x%08X"), *UE::ZoneGraph::Helpers::GetTagMaskString(Mask, TEXT(", ")), Mask.GetValue()), ScreenLoc.X, ScreenLoc.Y);
-		}
+			const float DistanceSq = FVector::DistSquared(ViewLocation, Zone.Bounds.GetCenter());
+			if (DistanceSq > DrawDistanceSq)
+			{
+				continue;
+			}
 
-		for (auto It = RegisteredLaneData.LaneToTrackingDataLookup.CreateConstIterator(); It; ++It)
-		{
-			const FCrowdTrackingLaneData& TrackingData = It->Value;
-			if (TrackingData.NumEntitiesOnLane > 0)
+			for (int32 LaneIdx = Zone.LanesBegin; LaneIdx < Zone.LanesEnd; LaneIdx++)
 			{
-				const int32 LaneIndex = It->Key;
+				const FZoneGraphLaneHandle LaneHandle(LaneIdx, RegisteredLaneData.DataHandle);
+				
 				FZoneGraphLaneLocation CenterLoc;
-				UE::ZoneGraph::Query::CalculateLocationAlongLaneFromRatio(*ZoneStorage, LaneIndex, 0.5f, CenterLoc);
-				const FVector ScreenLoc = Canvas->Project(CenterLoc.Position + ZOffset);
-				Canvas->DrawText(RenderFont, FString::Printf(TEXT("Num: %d"), TrackingData.NumEntitiesOnLane), ScreenLoc.X, ScreenLoc.Y);
+				UE::ZoneGraph::Query::CalculateLocationAlongLaneFromRatio(*ZoneStorage, LaneIdx, 0.5f, CenterLoc);
+
+				if (!InFrustum(CenterLoc.Position))
+				{
+					continue;
+				}
+				
+				const FVector ScreenLoc = Canvas->Project(CenterLoc.Position, /*bClampToNearPlane*/false);
+				
+				// Flags
+				if (bDisplayTags)
+				{
+					const FZoneGraphTagMask Mask = ZoneGraphAnnotationSubsystem->GetAnnotationTags(LaneHandle);
+					Canvas->DrawText(RenderFont, FString::Printf(TEXT("%s\n0x%08X"), *UE::ZoneGraph::Helpers::GetTagMaskString(Mask, TEXT(", ")), Mask.GetValue()), ScreenLoc.X, ScreenLoc.Y, 1.0f, 1.0f, FontInfo);
+				}
+
+				// Tracking
+				if (const FCrowdTrackingLaneData* TrackingData = RegisteredLaneData.LaneToTrackingDataLookup.Find(LaneIdx))
+				{
+					if (TrackingData->NumEntitiesOnLane > 0)
+					{
+						Canvas->DrawText(RenderFont, FString::Printf(TEXT("Num Entities: %d"), TrackingData->NumEntitiesOnLane), ScreenLoc.X, ScreenLoc.Y + 20, 1.0f, 1.0f, FontInfo);
+					}
+				}
 			}
 		}
 
+		// Waiting areas
 		for (const FCrowdWaitAreaData& WaitArea : RegisteredLaneData.WaitAreas)
 		{
-			for (const FCrowdWaitSlot& Slot : WaitArea.Slots)
+			if (WaitArea.Slots.Num() > 0)
 			{
-				if (Slot.bOccupied)
+				const float DistanceSq = FVector::DistSquared(ViewLocation, WaitArea.Slots[0].Position);
+				if (DistanceSq < DrawDistanceSq)
 				{
-					const FVector ScreenLoc = Canvas->Project(Slot.Position + ZOffset);
-					Canvas->DrawText(RenderFont, TEXT("OCCUPIED"), ScreenLoc.X, ScreenLoc.Y);
+					for (const FCrowdWaitSlot& Slot : WaitArea.Slots)
+					{
+						if (Slot.bOccupied)
+						{
+							if (!InFrustum(Slot.Position + ZOffset))
+							{
+								continue;
+							}
+
+							const FVector ScreenLoc = Canvas->Project(Slot.Position + ZOffset);
+							Canvas->DrawText(RenderFont, TEXT("OCCUPIED"), ScreenLoc.X, ScreenLoc.Y, 1.0f, 1.0f, FontInfo);
+						}
+					}
 				}
 			}
 		}
diff --git a/Engine/Plugins/AI/MassCrowd/Source/MassCrowd/Public/ZoneGraphCrowdLaneAnnotations.h b/Engine/Plugins/AI/MassCrowd/Source/MassCrowd/Public/ZoneGraphCrowdLaneAnnotations.h
--- a/Engine/Plugins/AI/MassCrowd/Source/MassCrowd/Public/ZoneGraphCrowdLaneAnnotations.h
+++ b/Engine/Plugins/AI/MassCrowd/Source/MassCrowd/Public/ZoneGraphCrowdLaneAnnotations.h
@@ -54,6 +54,9 @@ protected:
 	UPROPERTY(EditAnywhere, Category = CrowdLane)
 	FZoneGraphTag WaitingLaneTag;
 
+	UPROPERTY(EditAnywhere, Category = Debug)
+	bool bDisplayTags = false;
+
 	/** Array of queued events. */
 	TArray<FZoneGraphCrowdLaneStateChangeEvent> StateChangeEvents;
 
diff --git a/Engine/Plugins/Runtime/MassGameplay/Source/MassSpawner/Private/MassSpawner.cpp b/Engine/Plugins/Runtime/MassGameplay/Source/MassSpawner/Private/MassSpawner.cpp
--- a/Engine/Plugins/Runtime/MassGameplay/Source/MassSpawner/Private/MassSpawner.cpp
+++ b/Engine/Plugins/Runtime/MassGameplay/Source/MassSpawner/Private/MassSpawner.cpp
@@ -106,7 +106,10 @@ AMassSpawner::AMassSpawner()
 #endif // WITH_EDITOR
 
 	SetCanBeDamaged(false);
+
+#if UE_BUILD_SHIPPING || UE_BUILD_TEST
 	SetActorHiddenInGame(true);
+#endif
 
 	bAutoSpawnOnBeginPlay = true;
 	bOverrideSchematics = false;
diff --git a/Engine/Plugins/Runtime/ZoneGraphAnnotations/Source/ZoneGraphAnnotations/Private/Annotations/ZoneGraphDisturbanceAnnotation.cpp b/Engine/Plugins/Runtime/ZoneGraphAnnotations/Source/ZoneGraphAnnotations/Private/Annotations/ZoneGraphDisturbanceAnnotation.cpp
--- a/Engine/Plugins/Runtime/ZoneGraphAnnotations/Source/ZoneGraphAnnotations/Private/Annotations/ZoneGraphDisturbanceAnnotation.cpp
+++ b/Engine/Plugins/Runtime/ZoneGraphAnnotations/Source/ZoneGraphAnnotations/Private/Annotations/ZoneGraphDisturbanceAnnotation.cpp
@@ -70,6 +70,23 @@ void UZoneGraphDisturbanceAnnotation::TickAnnotation(const float DeltaTime, FZon
 		
 		UpdateAnnotationTags(AnnotationTagContainer);
 	}
+
+#if !UE_BUILD_SHIPPING && !UE_BUILD_TEST
+	if (bEnableDebugDrawing)
+	{
+		FVector ViewLocation = FVector::ZeroVector;
+		FRotator ViewRotation = FRotator::ZeroRotator;
+		GetFirstViewPoint(ViewLocation, ViewRotation);
+
+		// The debug draw caches an area around the current view. Update the debug drawing if the view moves far enough from the last snapshot.
+		const float UpdateThreshold = GetMaxDebugDrawDistance() * 0.25f;
+		if (FVector::DistSquared(LastDebugDrawLocation, ViewLocation) > FMath::Square(UpdateThreshold))
+		{
+			MarkRenderStateDirty();
+		}
+	}
+#endif
+
 }
 
 void UZoneGraphDisturbanceAnnotation::UpdateDangerLanes()
@@ -172,8 +189,9 @@ void UZoneGraphDisturbanceAnnotation::CalculateEscapeGraph(FZoneGraphDataEscapeG
 	TArray<FEscapeNode> DisturbanceHeap;
 	TArray<FZoneGraphLinkedLane> LinkedLanes;
 
-	static const float InvIdealSpanLength = 1.0f / FMath::Max(1.0f, IdealSpanLength);
-	
+	const float InvIdealSpanLength = 1.0f / FMath::Max(1.0f, IdealSpanLength);
+	float LowestLaneDanger = MAX_flt;
+
 	for (int32 EscapeLaneIndex = 0; EscapeLaneIndex < EscapeGraph.LanesToEscape.Num(); EscapeLaneIndex++)
 	{
 		FZoneGraphEscapeLaneAction& EscapeLane = EscapeGraph.LanesToEscape[EscapeLaneIndex];
@@ -188,6 +206,8 @@ void UZoneGraphDisturbanceAnnotation::CalculateEscapeGraph(FZoneGraphDataEscapeG
 		
 		FVector PrevPosition = ZoneStorage->LanePoints[Lane.PointsBegin];
 		float PrevDistance = 0.0f;
+
+		float AccumulatedDanger = 0.0f;
 		
 		for (uint8 SpanIndex = 0; SpanIndex < EscapeLane.SpanCount; SpanIndex++)
 		{
@@ -227,9 +247,13 @@ void UZoneGraphDisturbanceAnnotation::CalculateEscapeGraph(FZoneGraphDataEscapeG
 				Span.Danger = FMath::Max(Span.Danger, DangerCost);
 			}
 
+			AccumulatedDanger += Span.Danger;
+			
 			PrevPosition = LaneLocation.Position;
 			PrevDistance = Distance;
 		}
+
+		LowestLaneDanger = FMath::Min(LowestLaneDanger, AccumulatedDanger);
 		
 		EscapeLane.LaneLength = LaneLength;
 
@@ -263,6 +287,68 @@ void UZoneGraphDisturbanceAnnotation::CalculateEscapeGraph(FZoneGraphDataEscapeG
 		}
 	}
 
+	// If none of the lanes lead to safety, remove lowest danger lanes and try to connect again.
+	// This can happen for example when danger area covers a city block and crosswalks are not allowed as escape lanes.
+	if (DisturbanceHeap.IsEmpty())
+	{
+		// Remove all lanes that have accumulated danger value below the lowest accumulated danger found earlier. 
+		const float ExitDangerThreshold = LowestLaneDanger + KINDA_SMALL_NUMBER;
+		for (int32 EscapeLaneIndex = 0; EscapeLaneIndex < EscapeGraph.LanesToEscape.Num(); EscapeLaneIndex++)
+		{
+			FZoneGraphEscapeLaneAction& EscapeLane = EscapeGraph.LanesToEscape[EscapeLaneIndex];
+		
+			// Check if the lane leads out of the danger zone, and add as a starting point for the search.
+			float AccumulatedDanger = 0.0f;
+			for (uint8 SpanIndex = 0; SpanIndex < EscapeLane.SpanCount; SpanIndex += (EscapeLane.SpanCount - 1))
+			{
+				FZoneGraphEscapeLaneSpan& Span = EscapeLane.Spans[SpanIndex];
+				AccumulatedDanger += Span.Danger;
+			}
+
+			if (AccumulatedDanger < ExitDangerThreshold)
+			{
+				EscapeGraph.LanesToEscape.RemoveAt(EscapeLaneIndex);
+				EscapeLaneIndex--;
+			}
+		}
+
+		// Try to reconnect exits.
+		for (int32 EscapeLaneIndex = 0; EscapeLaneIndex < EscapeGraph.LanesToEscape.Num(); EscapeLaneIndex++)
+		{
+			FZoneGraphEscapeLaneAction& EscapeLane = EscapeGraph.LanesToEscape[EscapeLaneIndex];
+			const FZoneGraphLaneHandle LaneHandle(EscapeLane.LaneIndex, EscapeGraph.DataHandle);
+			const FZoneLaneData& Lane = ZoneStorage->Lanes[EscapeLane.LaneIndex];
+		
+			// Check if the lane leads out of the danger zone, and add as a starting point for the search.
+			for (uint8 SpanIndex = 0; SpanIndex < EscapeLane.SpanCount; SpanIndex += (EscapeLane.SpanCount - 1))
+			{
+				FZoneGraphEscapeLaneSpan& Span = EscapeLane.Spans[SpanIndex];
+
+				const EZoneLaneLinkType ExitLinkType = SpanIndex == 0 ? EZoneLaneLinkType::Incoming : EZoneLaneLinkType::Outgoing;
+				UE::ZoneGraph::Query::GetLinkedLanes(*ZoneStorage, LaneHandle, ExitLinkType, EZoneLaneLinkFlags::All, EZoneLaneLinkFlags::None, LinkedLanes);
+
+				for (const FZoneGraphLinkedLane& LinkedLane : LinkedLanes)
+				{
+					// If the linked lane is not part of the lanes to escape and passes the escape lanes filter, add this avoided lane as start for the search.
+					if (!EscapeGraph.LanesToEscapeLookup.Contains(LinkedLane.DestLane.Index))
+					{
+						const FZoneLaneData& DestLane = ZoneStorage->Lanes[LinkedLane.DestLane.Index];
+						if (EscapeLaneTags.Pass(DestLane.Tags))
+						{
+							DisturbanceHeap.Emplace(EscapeLaneIndex, /*EscapeCost*/0.0f, SpanIndex, LinkedLane.DestLane.Index, ExitLinkType);
+							Span.bLeadsToExit = true;
+							Span.ExitLinkType = ExitLinkType;
+							Span.EscapeCost = 0.0f;
+							Span.ExitLaneIndex = LinkedLane.DestLane.Index;
+							Span.bReverseLaneDirection = ExitLinkType == EZoneLaneLinkType::Incoming;
+							break;
+						}
+					}
+				}
+			}
+		}
+	}
+	
 	// Flood exit direction along the danger lanes.
 	// 
 	// NOTE:A lot of this code feels reversed because we're advancing from the fringe towards the center,
@@ -599,6 +685,16 @@ void UZoneGraphDisturbanceAnnotation::DebugDraw(FZoneGraphAnnotationSceneProxy*
 	static const FLinearColor EscapeCostHighColor(FColor(255, 61, 0));
 	static constexpr float ArrowSize = 15.0f;
 
+	FVector ViewLocation = FVector::ZeroVector;
+	FRotator ViewRotation = FRotator::ZeroRotator;
+	GetFirstViewPoint(ViewLocation, ViewRotation);
+
+	LastDebugDrawLocation = ViewLocation; 
+	
+	const float DrawDistance = GetMaxDebugDrawDistance();
+	const float DrawDistanceSq = FMath::Square(DrawDistance);
+	const float LabelDrawDistanceSq = FMath::Square(DrawDistance * 0.15f);
+
 	for (const FZoneGraphDataEscapeGraph& EscapeGraph : EscapeGraphs)
 	{
 		if (!EscapeGraph.bInUse)
@@ -610,6 +706,18 @@ void UZoneGraphDisturbanceAnnotation::DebugDraw(FZoneGraphAnnotationSceneProxy*
 		{
 			for (const FZoneGraphEscapeLaneAction& EscapeLane : EscapeGraph.LanesToEscape)
 			{
+				if (EscapeLane.SpanCount == 0)
+				{
+					continue;
+				}
+				const float DistanceSq = FVector::DistSquared(ViewLocation, EscapeLane.Spans[EscapeLane.SpanCount/2].Position);
+				if (DistanceSq > DrawDistanceSq)
+				{
+					continue;
+				}
+
+				const bool bAddLabel = DistanceSq < LabelDrawDistanceSq;
+				
 				const FZoneGraphLaneHandle LaneHandle(EscapeLane.LaneIndex, EscapeGraph.DataHandle);
 				
 				// Draw lane to indicate danger.
@@ -661,20 +769,26 @@ void UZoneGraphDisturbanceAnnotation::DebugDraw(FZoneGraphAnnotationSceneProxy*
 						DebugProxy->Lines.Emplace(ArrowPosition - ArrowDirection * ArrowSize + ArrowSide * ArrowSize, ArrowPosition + ArrowDirection * ArrowSize, MarkerColor.ToFColor(/*sRGB*/true), 2.0f);
 					}
 
-					if (Span.ExitLaneIndex == INDEX_NONE)
-					{
-						DebugProxy->Texts.Emplace(FString::Printf(TEXT("ERROR!")), ArrowPosition + FVector::UpVector * ArrowSize, FColor::Red);
-					}
-					else
+					if (bAddLabel)
 					{
-						DebugProxy->Texts.Emplace(FString::Printf(TEXT("%d->%d\nescapeCost: %.1f\ndanger: %.1f\ntags: %s"), EscapeLane.LaneIndex, Span.ExitLaneIndex, Span.EscapeCost, Span.Danger, *TagString), ArrowPosition + FVector::UpVector * ArrowSize, FColor::Red);
+						if (Span.ExitLaneIndex == INDEX_NONE)
+						{
+							DebugProxy->Texts.Emplace(FString::Printf(TEXT("ERROR!")), ArrowPosition + FVector::UpVector * ArrowSize, FColor::Red);
+						}
+						else
+						{
+							DebugProxy->Texts.Emplace(FString::Printf(TEXT("%d->%d\nescapeCost: %.1f\ndanger: %.1f\ntags: %s"), EscapeLane.LaneIndex, Span.ExitLaneIndex, Span.EscapeCost, Span.Danger, *TagString), ArrowPosition + FVector::UpVector * ArrowSize, FColor::Red);
+						}
 					}
 
 					if (Span.bLeadsToExit)
 					{
 						const FVector Point = SpanIndex == 0 ? PrevLocation.Position : SplitLocation.Position + ZOffset;
 						DebugProxy->Lines.Emplace(Point, Point + FVector::UpVector * ArrowSize, FColor::Blue, 2.0f);
-						DebugProxy->Texts.Emplace(FString(TEXT("EXIT")), Point + FVector::UpVector * ArrowSize, FColor::Blue);
+						if (bAddLabel)
+						{
+							DebugProxy->Texts.Emplace(FString(TEXT("EXIT")), Point + FVector::UpVector * ArrowSize, FColor::Blue);
+						}
 					}
 					
 					PrevLocation = SplitLocation;
@@ -685,6 +799,12 @@ void UZoneGraphDisturbanceAnnotation::DebugDraw(FZoneGraphAnnotationSceneProxy*
 
 	for (const FZoneGraphDisturbanceArea& Danger : Dangers)
 	{
+		const float DistanceSq = FVector::DistSquared(ViewLocation, Danger.Position);
+		if (DistanceSq > DrawDistanceSq)
+		{
+			continue;
+		}
+
 		const FBox QueryBounds = FBox::BuildAABB(Danger.Position, FVector(Danger.Radius));
 		DebugProxy->Boxes.Emplace(QueryBounds, FColor::Red);
 		DebugProxy->Lines.Emplace(Danger.Position - FVector(ArrowSize, 0, 0), Danger.Position + FVector(ArrowSize, 0, 0), FColor::Red, 2.0f);
@@ -694,6 +814,12 @@ void UZoneGraphDisturbanceAnnotation::DebugDraw(FZoneGraphAnnotationSceneProxy*
 
 	for (const FZoneGraphObstacleDisturbanceArea& Obstacle : Obstacles)
 	{
+		const float DistanceSq = FVector::DistSquared(ViewLocation, Obstacle.Position);
+		if (DistanceSq > DrawDistanceSq)
+		{
+			continue;
+		}
+
 		const FColor ObstacleColor = FColor::Blue;
 		const FBox QueryBounds = FBox::BuildAABB(Obstacle.Position, FVector(Obstacle.Radius));
 		DebugProxy->Boxes.Emplace(QueryBounds, ObstacleColor);
@@ -707,7 +833,7 @@ void UZoneGraphDisturbanceAnnotation::DebugDraw(FZoneGraphAnnotationSceneProxy*
 void UZoneGraphDisturbanceAnnotation::DebugDrawCanvas(UCanvas* Canvas, APlayerController*)
 {
 	const FColor OldDrawColor = Canvas->DrawColor;
-	const FFontRenderInfo FontInfo = Canvas->CreateFontRenderInfo(false, true);
+	const FFontRenderInfo FontInfo = Canvas->CreateFontRenderInfo(true, true);
 	const UFont* RenderFont = GEngine->GetSmallFont();
 	const float LineHeight = RenderFont->GetMaxCharHeight() * 1.2f;
 
@@ -715,8 +841,12 @@ void UZoneGraphDisturbanceAnnotation::DebugDrawCanvas(UCanvas* Canvas, APlayerCo
 
 	constexpr float PosX = 40;
 	float PosY = 200;
-	for (const FZoneGraphDisturbanceArea& Danger : Dangers)
+
+	const int32 NumLines = FMath::Min(Dangers.Num(), 15);
+	
+	for (int32 Index = 0; Index < NumLines; Index++)
 	{
+		const FZoneGraphDisturbanceArea& Danger = Dangers[Index];
 		FString Text = FString::Printf(TEXT("Danger %.1f"), Danger.Duration);
 		
 		Canvas->SetDrawColor(FColor::Red);
@@ -724,6 +854,13 @@ void UZoneGraphDisturbanceAnnotation::DebugDrawCanvas(UCanvas* Canvas, APlayerCo
 
 		PosY -= LineHeight;
 	}
+	
+	if (NumLines < Dangers.Num())
+	{
+		FString Text = FString::Printf(TEXT("%d more dangers..."), Dangers.Num() - NumLines);
+		Canvas->SetDrawColor(FColor::Red);
+		Canvas->DrawText(RenderFont, Text, PosX, PosY, 1, 1, FontInfo);
+	}
 
 	Canvas->SetDrawColor(OldDrawColor);
 }
diff --git a/Engine/Plugins/Runtime/ZoneGraphAnnotations/Source/ZoneGraphAnnotations/Private/ZoneGraphAnnotationComponent.cpp b/Engine/Plugins/Runtime/ZoneGraphAnnotations/Source/ZoneGraphAnnotations/Private/ZoneGraphAnnotationComponent.cpp
--- a/Engine/Plugins/Runtime/ZoneGraphAnnotations/Source/ZoneGraphAnnotations/Private/ZoneGraphAnnotationComponent.cpp
+++ b/Engine/Plugins/Runtime/ZoneGraphAnnotations/Source/ZoneGraphAnnotations/Private/ZoneGraphAnnotationComponent.cpp
@@ -6,12 +6,10 @@
 #include "ZoneGraphAnnotationSubsystem.h"
 #include "ZoneGraphSubsystem.h"
 #include "ZoneGraphDelegates.h"
+#include "ZoneGraphSettings.h"
 #include "Debug/DebugDrawService.h"
-
-#if WITH_EDITOR
-#include "Editor.h"
-#endif
-
+#include "Engine/DebugCameraController.h"
+#include "Engine/LocalPlayer.h"
 
 #if UE_ENABLE_DEBUG_DRAWING
 //////////////////////////////////////////////////////////////////////////
@@ -151,6 +149,8 @@ void UZoneGraphAnnotationComponent::PostSubsystemsInitialized()
 
 void UZoneGraphAnnotationComponent::OnUnregister()
 {
+	Super::OnUnregister();
+
 	UWorld* World = GetWorld();
 	if (World == nullptr)
 	{
@@ -181,8 +181,6 @@ void UZoneGraphAnnotationComponent::OnUnregister()
 #if UE_ENABLE_DEBUG_DRAWING
 	UDebugDrawService::Unregister(CanvasDebugDrawDelegateHandle);
 #endif
-	
-	Super::OnUnregister();
 }
 
 void UZoneGraphAnnotationComponent::OnPostZoneGraphDataAdded(const AZoneGraphData* ZoneGraphData)
@@ -208,6 +206,50 @@ void UZoneGraphAnnotationComponent::OnPreZoneGraphDataRemoved(const AZoneGraphDa
 }
 
 #if UE_ENABLE_DEBUG_DRAWING
+void UZoneGraphAnnotationComponent::GetFirstViewPoint(FVector& ViewLocation, FRotator& ViewRotation) const
+{
+	ViewLocation = FVector::ZeroVector;
+	ViewRotation = FRotator::ZeroRotator;
+
+	bool bFound = false;
+	
+	if (const UWorld* World = GetWorld())
+	{
+		// Now go through all current player controllers and add if they do not exist
+		for (FConstPlayerControllerIterator PlayerIterator = World->GetPlayerControllerIterator(); PlayerIterator; ++PlayerIterator)
+		{
+			if (const APlayerController* PlayerController = (*PlayerIterator).Get())
+			{
+				PlayerController->GetPlayerViewPoint(ViewLocation, ViewRotation);
+				bFound = true;
+				break;
+			}
+		}
+
+		// spectator mode
+		if (!bFound)
+		{
+			for (FLocalPlayerIterator It(GEngine, const_cast<UWorld*>(World)); It; ++It)
+			{
+				ADebugCameraController* SpectatorPC = Cast<ADebugCameraController>(It->PlayerController);
+				if (SpectatorPC)
+				{
+					SpectatorPC->GetPlayerViewPoint(ViewLocation, ViewRotation);
+					bFound = true;
+					break;
+				}
+			}
+		}
+
+	}
+}
+
+float UZoneGraphAnnotationComponent::GetMaxDebugDrawDistance() const
+{
+	const UZoneGraphSettings* ZoneGraphSettings = GetDefault<UZoneGraphSettings>();
+	return ZoneGraphSettings ? ZoneGraphSettings->GetShapeMaxDrawDistance() : 0.0f;
+}
+
 FDebugRenderSceneProxy* UZoneGraphAnnotationComponent::CreateDebugSceneProxy()
 {
 	FZoneGraphAnnotationSceneProxy* DebugProxy = new FZoneGraphAnnotationSceneProxy(*this);
diff --git a/Engine/Plugins/Runtime/ZoneGraphAnnotations/Source/ZoneGraphAnnotations/Public/Annotations/ZoneGraphDisturbanceAnnotation.h b/Engine/Plugins/Runtime/ZoneGraphAnnotations/Source/ZoneGraphAnnotations/Public/Annotations/ZoneGraphDisturbanceAnnotation.h
--- a/Engine/Plugins/Runtime/ZoneGraphAnnotations/Source/ZoneGraphAnnotations/Public/Annotations/ZoneGraphDisturbanceAnnotation.h
+++ b/Engine/Plugins/Runtime/ZoneGraphAnnotations/Source/ZoneGraphAnnotations/Public/Annotations/ZoneGraphDisturbanceAnnotation.h
@@ -182,6 +182,8 @@ protected:
 #if !UE_BUILD_SHIPPING && !UE_BUILD_TEST
 	virtual void DebugDraw(FZoneGraphAnnotationSceneProxy* DebugProxy) override;
 	virtual void DebugDrawCanvas(UCanvas* Canvas, APlayerController*) override;
+
+	FVector LastDebugDrawLocation = FVector::ZeroVector; 
 #endif
 
 	void CalculateEscapeGraph(FZoneGraphDataEscapeGraph& EscapeGraph);
diff --git a/Engine/Plugins/Runtime/ZoneGraphAnnotations/Source/ZoneGraphAnnotations/Public/ZoneGraphAnnotationComponent.h b/Engine/Plugins/Runtime/ZoneGraphAnnotations/Source/ZoneGraphAnnotations/Public/ZoneGraphAnnotationComponent.h
--- a/Engine/Plugins/Runtime/ZoneGraphAnnotations/Source/ZoneGraphAnnotations/Public/ZoneGraphAnnotationComponent.h
+++ b/Engine/Plugins/Runtime/ZoneGraphAnnotations/Source/ZoneGraphAnnotations/Public/ZoneGraphAnnotationComponent.h
@@ -8,7 +8,7 @@
 #include "ZoneGraphAnnotationComponent.generated.h"
 
 class UZoneGraphAnnotationSubsystem;
-class UMassZoneGraphAnnotationComponent;
+class UZoneGraphAnnotationComponent;
 class UCanvas;
 class AZoneGraphData;
 struct FInstancedStructStream;
@@ -58,6 +58,12 @@ public:
 	virtual void PreZoneGraphDataRemoved(const AZoneGraphData& ZoneGraphData) {}
 	
 #if UE_ENABLE_DEBUG_DRAWING
+	/** Returns first view point (player controller or debug camera) */
+	void GetFirstViewPoint(FVector& ViewLocation, FRotator& ViewRotation) const;
+
+	/** Returns ZoneGraph max debug draw distance. */
+	float GetMaxDebugDrawDistance() const;
+
 	/** Called when scene proxy is rebuilt. */
 	virtual void DebugDraw(FZoneGraphAnnotationSceneProxy* DebugProxy) {}
 
diff --git a/Engine/Plugins/Runtime/ZoneGraphAnnotations/Source/ZoneGraphAnnotations/ZoneGraphAnnotations.Build.cs b/Engine/Plugins/Runtime/ZoneGraphAnnotations/Source/ZoneGraphAnnotations/ZoneGraphAnnotations.Build.cs
--- a/Engine/Plugins/Runtime/ZoneGraphAnnotations/Source/ZoneGraphAnnotations/ZoneGraphAnnotations.Build.cs
+++ b/Engine/Plugins/Runtime/ZoneGraphAnnotations/Source/ZoneGraphAnnotations/ZoneGraphAnnotations.Build.cs
@@ -31,6 +31,12 @@ namespace UnrealBuildTool.Rules
 					"ZoneGraphDebug"
 				}
 			);
+			
+			if (Target.bBuildEditor == true)
+			{
+				PrivateDependencyModuleNames.Add("EditorFramework");
+				PrivateDependencyModuleNames.Add("UnrealEd");
+			}
 		}
 	}
 }
\ No newline at end of file
