ZoneGraphAnnotationComponen Unity compile fix

diff --git a/Engine/Plugins/Runtime/ZoneGraphAnnotations/Source/ZoneGraphAnnotations/Private/ZoneGraphAnnotationComponent.cpp b/Engine/Plugins/Runtime/ZoneGraphAnnotations/Source/ZoneGraphAnnotations/Private/ZoneGraphAnnotationComponent.cpp
--- a/Engine/Plugins/Runtime/ZoneGraphAnnotations/Source/ZoneGraphAnnotations/Private/ZoneGraphAnnotationComponent.cpp
+++ b/Engine/Plugins/Runtime/ZoneGraphAnnotations/Source/ZoneGraphAnnotations/Private/ZoneGraphAnnotationComponent.cpp
@@ -10,6 +10,7 @@
 #include "Debug/DebugDrawService.h"
 #include "Engine/DebugCameraController.h"
 #include "Engine/LocalPlayer.h"
+#include "UnrealEngine.h"
 
 #if UE_ENABLE_DEBUG_DRAWING
 //////////////////////////////////////////////////////////////////////////
