ZoneGraph: Fix crash in flee annotation
- fixed crash (lookup out of date) in flee fallback behavior
- fixe left overd issues from an older review 19243357

diff --git a/Engine/Plugins/AI/MassCrowd/Source/MassCrowd/Private/ZoneGraphCrowdLaneAnnotations.cpp b/Engine/Plugins/AI/MassCrowd/Source/MassCrowd/Private/ZoneGraphCrowdLaneAnnotations.cpp
--- a/Engine/Plugins/AI/MassCrowd/Source/MassCrowd/Private/ZoneGraphCrowdLaneAnnotations.cpp
+++ b/Engine/Plugins/AI/MassCrowd/Source/MassCrowd/Private/ZoneGraphCrowdLaneAnnotations.cpp
@@ -205,7 +205,7 @@ void UZoneGraphCrowdLaneAnnotations::DebugDrawCanvas(UCanvas* Canvas, APlayerCon
 	const FColor OldDrawColor = Canvas->DrawColor;
 	const UFont* RenderFont = GEngine->GetSmallFont();
 
-	const FFontRenderInfo FontInfo = Canvas->CreateFontRenderInfo(true, true);
+	const FFontRenderInfo FontInfo = Canvas->CreateFontRenderInfo(/*bClipText*/true, /*bEnableShadow*/true);
 
 	Canvas->SetDrawColor(FColor::White);
 	static const FVector ZOffset(0, 0, 35.0f);
@@ -265,7 +265,7 @@ void UZoneGraphCrowdLaneAnnotations::DebugDrawCanvas(UCanvas* Canvas, APlayerCon
 				if (bDisplayTags)
 				{
 					const FZoneGraphTagMask Mask = ZoneGraphAnnotationSubsystem->GetAnnotationTags(LaneHandle);
-					Canvas->DrawText(RenderFont, FString::Printf(TEXT("%s\n0x%08X"), *UE::ZoneGraph::Helpers::GetTagMaskString(Mask, TEXT(", ")), Mask.GetValue()), ScreenLoc.X, ScreenLoc.Y, 1.0f, 1.0f, FontInfo);
+					Canvas->DrawText(RenderFont, FString::Printf(TEXT("%s\n0x%08X"), *UE::ZoneGraph::Helpers::GetTagMaskString(Mask, TEXT(", ")), Mask.GetValue()), ScreenLoc.X, ScreenLoc.Y, /*XScale*/1.0f, /*YScale*/1.0f, FontInfo);
 				}
 
 				// Tracking
@@ -273,7 +273,7 @@ void UZoneGraphCrowdLaneAnnotations::DebugDrawCanvas(UCanvas* Canvas, APlayerCon
 				{
 					if (TrackingData->NumEntitiesOnLane > 0)
 					{
-						Canvas->DrawText(RenderFont, FString::Printf(TEXT("Num Entities: %d"), TrackingData->NumEntitiesOnLane), ScreenLoc.X, ScreenLoc.Y + 20, 1.0f, 1.0f, FontInfo);
+						Canvas->DrawText(RenderFont, FString::Printf(TEXT("Num Entities: %d"), TrackingData->NumEntitiesOnLane), ScreenLoc.X, ScreenLoc.Y + 20, /*XScale*/1.0f, /*YScale*/1.0f, FontInfo);
 					}
 				}
 			}
@@ -297,7 +297,7 @@ void UZoneGraphCrowdLaneAnnotations::DebugDrawCanvas(UCanvas* Canvas, APlayerCon
 							}
 
 							const FVector ScreenLoc = Canvas->Project(Slot.Position + ZOffset);
-							Canvas->DrawText(RenderFont, TEXT("OCCUPIED"), ScreenLoc.X, ScreenLoc.Y, 1.0f, 1.0f, FontInfo);
+							Canvas->DrawText(RenderFont, TEXT("OCCUPIED"), ScreenLoc.X, ScreenLoc.Y, /*XScale*/1.0f, /*YScale*/1.0f, FontInfo);
 						}
 					}
 				}
diff --git a/Engine/Plugins/Runtime/ZoneGraphAnnotations/Source/ZoneGraphAnnotations/Private/Annotations/ZoneGraphDisturbanceAnnotation.cpp b/Engine/Plugins/Runtime/ZoneGraphAnnotations/Source/ZoneGraphAnnotations/Private/Annotations/ZoneGraphDisturbanceAnnotation.cpp
--- a/Engine/Plugins/Runtime/ZoneGraphAnnotations/Source/ZoneGraphAnnotations/Private/Annotations/ZoneGraphDisturbanceAnnotation.cpp
+++ b/Engine/Plugins/Runtime/ZoneGraphAnnotations/Source/ZoneGraphAnnotations/Private/Annotations/ZoneGraphDisturbanceAnnotation.cpp
@@ -119,16 +119,17 @@ void UZoneGraphDisturbanceAnnotation::UpdateDangerLanes()
 				FZoneGraphDataEscapeGraph& EscapeGraph = EscapeGraphs[LaneHandle.DataHandle.Index];
 				check(EscapeGraph.bInUse);
 				
-				// Add new lane if not added already.				
-				if (!EscapeGraph.LanesToEscapeLookup.Contains(LaneHandle.Index))
+				// Add new lane if not added already.
+				if (const int32* EscapeLaneIndex = EscapeGraph.LanesToEscapeLookup.Find(LaneHandle.Index))
 				{
-					const int32 Index = EscapeGraph.LanesToEscape.Emplace(LaneHandle.Index);
-					EscapeGraph.LanesToEscapeLookup.Add(LaneHandle.Index, Index);
-					EscapeGraph.LanesToEscape[Index].Tags.Add(Tag);
+					// Update existing
+					EscapeGraph.LanesToEscape[*EscapeLaneIndex].Tags.Add(Tag);
 				}
 				else
 				{
-					const int32 Index = EscapeGraph.LanesToEscapeLookup[LaneHandle.Index];
+					// Add new
+					const int32 Index = EscapeGraph.LanesToEscape.Emplace(LaneHandle.Index);
+					EscapeGraph.LanesToEscapeLookup.Add(LaneHandle.Index, Index);
 					EscapeGraph.LanesToEscape[Index].Tags.Add(Tag);
 				}
 			}
@@ -189,6 +190,42 @@ void UZoneGraphDisturbanceAnnotation::CalculateEscapeGraph(FZoneGraphDataEscapeG
 	TArray<FEscapeNode> DisturbanceHeap;
 	TArray<FZoneGraphLinkedLane> LinkedLanes;
 
+
+	auto AddSearchStart = [this, ZoneStorage, &LinkedLanes, &EscapeGraph, &DisturbanceHeap](const int32 EscapeLaneIndex)
+	{
+		FZoneGraphEscapeLaneAction& EscapeLane = EscapeGraph.LanesToEscape[EscapeLaneIndex];
+		const FZoneGraphLaneHandle LaneHandle(EscapeLane.LaneIndex, EscapeGraph.DataHandle);
+
+		// Check if the lane leads out of the danger zone, and add as a starting point for the search.
+		for (uint8 SpanIndex = 0; SpanIndex < EscapeLane.SpanCount; SpanIndex += (EscapeLane.SpanCount - 1))
+		{
+			FZoneGraphEscapeLaneSpan& Span = EscapeLane.Spans[SpanIndex];
+
+			const EZoneLaneLinkType ExitLinkType = SpanIndex == 0 ? EZoneLaneLinkType::Incoming : EZoneLaneLinkType::Outgoing;
+			UE::ZoneGraph::Query::GetLinkedLanes(*ZoneStorage, LaneHandle, ExitLinkType, EZoneLaneLinkFlags::All, EZoneLaneLinkFlags::None, LinkedLanes);
+
+			for (const FZoneGraphLinkedLane& LinkedLane : LinkedLanes)
+			{
+				// If the linked lane is not part of the lanes to escape and passes the escape lanes filter, add this avoided lane as start for the search.
+				if (!EscapeGraph.LanesToEscapeLookup.Contains(LinkedLane.DestLane.Index))
+				{
+					const FZoneLaneData& DestLane = ZoneStorage->Lanes[LinkedLane.DestLane.Index];
+					if (EscapeLaneTags.Pass(DestLane.Tags))
+					{
+						DisturbanceHeap.Emplace(EscapeLaneIndex, /*EscapeCost*/0.0f, SpanIndex, LinkedLane.DestLane.Index, ExitLinkType);
+						Span.bLeadsToExit = true;
+						Span.ExitLinkType = ExitLinkType;
+						Span.EscapeCost = 0.0f;
+						Span.ExitLaneIndex = LinkedLane.DestLane.Index;
+						Span.bReverseLaneDirection = ExitLinkType == EZoneLaneLinkType::Incoming;
+						break;
+					}
+				}
+			}
+		}				
+	};
+
+	
 	const float InvIdealSpanLength = 1.0f / FMath::Max(1.0f, IdealSpanLength);
 	float LowestLaneDanger = MAX_flt;
 
@@ -257,39 +294,12 @@ void UZoneGraphDisturbanceAnnotation::CalculateEscapeGraph(FZoneGraphDataEscapeG
 		
 		EscapeLane.LaneLength = LaneLength;
 
-		// Check if the lane leads out of the danger zone, and add as a starting point for the search.
-		// The first span covers the start of the lane, and second one the end.
-		for (uint8 SpanIndex = 0; SpanIndex < EscapeLane.SpanCount; SpanIndex += (EscapeLane.SpanCount - 1))
-		{
-			FZoneGraphEscapeLaneSpan& Span = EscapeLane.Spans[SpanIndex];
-
-			const EZoneLaneLinkType ExitLinkType = SpanIndex == 0 ? EZoneLaneLinkType::Incoming : EZoneLaneLinkType::Outgoing;
-			UE::ZoneGraph::Query::GetLinkedLanes(*ZoneStorage, LaneHandle, ExitLinkType, EZoneLaneLinkFlags::All, EZoneLaneLinkFlags::None, LinkedLanes);
-
-			for (const FZoneGraphLinkedLane& LinkedLane : LinkedLanes)
-			{
-				// If the linked lane is not part of the lanes to escape and passes the escape lanes filter, add this avoided lane as start for the search.
-				if (!EscapeGraph.LanesToEscapeLookup.Contains(LinkedLane.DestLane.Index))
-				{
-					const FZoneLaneData& DestLane = ZoneStorage->Lanes[LinkedLane.DestLane.Index];
-					if (EscapeLaneTags.Pass(DestLane.Tags))
-					{
-						DisturbanceHeap.Emplace(EscapeLaneIndex, /*EscapeCost*/0.0f, SpanIndex, LinkedLane.DestLane.Index, ExitLinkType);
-						Span.bLeadsToExit = true;
-						Span.ExitLinkType = ExitLinkType;
-						Span.EscapeCost = 0.0f;
-						Span.ExitLaneIndex = LinkedLane.DestLane.Index;
-						Span.bReverseLaneDirection = ExitLinkType == EZoneLaneLinkType::Incoming;
-						break;
-					}
-				}
-			}
-		}
+		AddSearchStart(EscapeLaneIndex);
 	}
 
 	// If none of the lanes lead to safety, remove lowest danger lanes and try to connect again.
 	// This can happen for example when danger area covers a city block and crosswalks are not allowed as escape lanes.
-	if (DisturbanceHeap.IsEmpty())
+	if (DisturbanceHeap.IsEmpty() && EscapeGraph.LanesToEscape.Num() > 0)
 	{
 		// Remove all lanes that have accumulated danger value below the lowest accumulated danger found earlier. 
 		const float ExitDangerThreshold = LowestLaneDanger + KINDA_SMALL_NUMBER;
@@ -297,9 +307,8 @@ void UZoneGraphDisturbanceAnnotation::CalculateEscapeGraph(FZoneGraphDataEscapeG
 		{
 			FZoneGraphEscapeLaneAction& EscapeLane = EscapeGraph.LanesToEscape[EscapeLaneIndex];
 		
-			// Check if the lane leads out of the danger zone, and add as a starting point for the search.
 			float AccumulatedDanger = 0.0f;
-			for (uint8 SpanIndex = 0; SpanIndex < EscapeLane.SpanCount; SpanIndex += (EscapeLane.SpanCount - 1))
+			for (uint8 SpanIndex = 0; SpanIndex < EscapeLane.SpanCount; SpanIndex++)
 			{
 				FZoneGraphEscapeLaneSpan& Span = EscapeLane.Spans[SpanIndex];
 				AccumulatedDanger += Span.Danger;
@@ -312,40 +321,18 @@ void UZoneGraphDisturbanceAnnotation::CalculateEscapeGraph(FZoneGraphDataEscapeG
 			}
 		}
 
-		// Try to reconnect exits.
+		// Rebuild LanesToEscapeLookup
+		EscapeGraph.LanesToEscapeLookup.Reset();
 		for (int32 EscapeLaneIndex = 0; EscapeLaneIndex < EscapeGraph.LanesToEscape.Num(); EscapeLaneIndex++)
 		{
-			FZoneGraphEscapeLaneAction& EscapeLane = EscapeGraph.LanesToEscape[EscapeLaneIndex];
-			const FZoneGraphLaneHandle LaneHandle(EscapeLane.LaneIndex, EscapeGraph.DataHandle);
-			const FZoneLaneData& Lane = ZoneStorage->Lanes[EscapeLane.LaneIndex];
-		
-			// Check if the lane leads out of the danger zone, and add as a starting point for the search.
-			for (uint8 SpanIndex = 0; SpanIndex < EscapeLane.SpanCount; SpanIndex += (EscapeLane.SpanCount - 1))
-			{
-				FZoneGraphEscapeLaneSpan& Span = EscapeLane.Spans[SpanIndex];
-
-				const EZoneLaneLinkType ExitLinkType = SpanIndex == 0 ? EZoneLaneLinkType::Incoming : EZoneLaneLinkType::Outgoing;
-				UE::ZoneGraph::Query::GetLinkedLanes(*ZoneStorage, LaneHandle, ExitLinkType, EZoneLaneLinkFlags::All, EZoneLaneLinkFlags::None, LinkedLanes);
+			const FZoneGraphEscapeLaneAction& EscapeLane = EscapeGraph.LanesToEscape[EscapeLaneIndex];
+			EscapeGraph.LanesToEscapeLookup.Add(EscapeLane.LaneIndex, EscapeLaneIndex);
+		}
 
-				for (const FZoneGraphLinkedLane& LinkedLane : LinkedLanes)
-				{
-					// If the linked lane is not part of the lanes to escape and passes the escape lanes filter, add this avoided lane as start for the search.
-					if (!EscapeGraph.LanesToEscapeLookup.Contains(LinkedLane.DestLane.Index))
-					{
-						const FZoneLaneData& DestLane = ZoneStorage->Lanes[LinkedLane.DestLane.Index];
-						if (EscapeLaneTags.Pass(DestLane.Tags))
-						{
-							DisturbanceHeap.Emplace(EscapeLaneIndex, /*EscapeCost*/0.0f, SpanIndex, LinkedLane.DestLane.Index, ExitLinkType);
-							Span.bLeadsToExit = true;
-							Span.ExitLinkType = ExitLinkType;
-							Span.EscapeCost = 0.0f;
-							Span.ExitLaneIndex = LinkedLane.DestLane.Index;
-							Span.bReverseLaneDirection = ExitLinkType == EZoneLaneLinkType::Incoming;
-							break;
-						}
-					}
-				}
-			}
+		// Try again to add the search start locations.
+		for (int32 EscapeLaneIndex = 0; EscapeLaneIndex < EscapeGraph.LanesToEscape.Num(); EscapeLaneIndex++)
+		{
+			AddSearchStart(EscapeLaneIndex);
 		}
 	}
 	
