Updated some DebugDrawComponent and ZoneGraphAnnotations to rely on UE_ENABLE_DEBUG_DRAWING
#rnx

diff --git a/Engine/Plugins/AI/MassAI/Source/MassNavigationEditor/Private/MassNavigationTestingActor.cpp b/Engine/Plugins/AI/MassAI/Source/MassNavigationEditor/Private/MassNavigationTestingActor.cpp
--- a/Engine/Plugins/AI/MassAI/Source/MassNavigationEditor/Private/MassNavigationTestingActor.cpp
+++ b/Engine/Plugins/AI/MassAI/Source/MassNavigationEditor/Private/MassNavigationTestingActor.cpp
@@ -8,7 +8,7 @@
 #include "ZoneGraphRenderingUtilities.h"
 #include "ZoneGraphData.h"
 
-#if !UE_BUILD_SHIPPING && !UE_BUILD_TEST
+#if UE_ENABLE_DEBUG_DRAWING
 
 //////////////////////////////////////////////////////////////////////////
 // FMassNavigationTestingSceneProxy
@@ -39,7 +39,7 @@ uint32 FMassNavigationTestingSceneProxy::GetMemoryFootprint(void) const
 	return sizeof(*this) + FDebugRenderSceneProxy::GetAllocatedSize();
 }
 
-#endif // !UE_BUILD_SHIPPING && !UE_BUILD_TEST
+#endif // UE_ENABLE_DEBUG_DRAWING
 
 //////////////////////////////////////////////////////////////////////////
 // UMassNavigationTestingComponent
diff --git a/Engine/Plugins/AI/MassAI/Source/MassNavigationEditor/Private/MassNavigationTestingActor.h b/Engine/Plugins/AI/MassAI/Source/MassNavigationEditor/Private/MassNavigationTestingActor.h
--- a/Engine/Plugins/AI/MassAI/Source/MassNavigationEditor/Private/MassNavigationTestingActor.h
+++ b/Engine/Plugins/AI/MassAI/Source/MassNavigationEditor/Private/MassNavigationTestingActor.h
@@ -14,7 +14,7 @@ class UZoneGraphTestingComponent;
 class UZoneGraphSubsystem;
 class AMassNavigationTestingActor;
 
-#if !UE_BUILD_SHIPPING && !UE_BUILD_TEST
+#if UE_ENABLE_DEBUG_DRAWING
 class MASSNAVIGATIONEDITOR_API FMassNavigationTestingSceneProxy final : public FDebugRenderSceneProxy
 {
 public:
@@ -24,7 +24,7 @@ public:
 	virtual FPrimitiveViewRelevance GetViewRelevance(const FSceneView* View) const override;
 	virtual uint32 GetMemoryFootprint(void) const override;
 };
-#endif
+#endif // UE_ENABLE_DEBUG_DRAWING
 
 /** Component for testing MassMovement functionality. */
 UCLASS(ClassGroup = Debug)
diff --git a/Engine/Plugins/AI/MassCrowd/Source/MassCrowd/Private/ZoneGraphCrowdLaneAnnotations.cpp b/Engine/Plugins/AI/MassCrowd/Source/MassCrowd/Private/ZoneGraphCrowdLaneAnnotations.cpp
--- a/Engine/Plugins/AI/MassCrowd/Source/MassCrowd/Private/ZoneGraphCrowdLaneAnnotations.cpp
+++ b/Engine/Plugins/AI/MassCrowd/Source/MassCrowd/Private/ZoneGraphCrowdLaneAnnotations.cpp
@@ -10,9 +10,9 @@
 #include "MassNavigationTypes.h"
 #include "VisualLogger/VisualLogger.h"
 
-#if !UE_BUILD_SHIPPING && !UE_BUILD_TEST
+#if UE_ENABLE_DEBUG_DRAWING
 #include "Engine/Canvas.h"
-#endif // !UE_BUILD_SHIPPING && !UE_BUILD_TEST
+#endif // UE_ENABLE_DEBUG_DRAWING
 
 void UZoneGraphCrowdLaneAnnotations::PostSubsystemsInitialized()
 {
@@ -84,15 +84,15 @@ void UZoneGraphCrowdLaneAnnotations::TickAnnotation(const float DeltaTime, FZone
 	}
 	StateChangeEvents.Reset();
 
-#if !UE_BUILD_SHIPPING && !UE_BUILD_TEST
+#if UE_ENABLE_DEBUG_DRAWING
 	if (bEnableDebugDrawing)
 	{
 		MarkRenderStateDirty();
 	}
-#endif
+#endif // UE_ENABLE_DEBUG_DRAWING
 }
 
-#if !UE_BUILD_SHIPPING && !UE_BUILD_TEST
+#if UE_ENABLE_DEBUG_DRAWING
 void UZoneGraphCrowdLaneAnnotations::DebugDraw(FZoneGraphAnnotationSceneProxy* DebugProxy)
 {
 	UZoneGraphSubsystem* ZoneGraph = UWorld::GetSubsystem<UZoneGraphSubsystem>(GetWorld());
@@ -308,4 +308,4 @@ void UZoneGraphCrowdLaneAnnotations::DebugDrawCanvas(UCanvas* Canvas, APlayerCon
 	Canvas->SetDrawColor(OldDrawColor);
 }
 
-#endif // !UE_BUILD_SHIPPING && !UE_BUILD_TEST
+#endif // UE_ENABLE_DEBUG_DRAWING
diff --git a/Engine/Plugins/AI/MassCrowd/Source/MassCrowd/Public/ZoneGraphCrowdLaneAnnotations.h b/Engine/Plugins/AI/MassCrowd/Source/MassCrowd/Public/ZoneGraphCrowdLaneAnnotations.h
--- a/Engine/Plugins/AI/MassCrowd/Source/MassCrowd/Public/ZoneGraphCrowdLaneAnnotations.h
+++ b/Engine/Plugins/AI/MassCrowd/Source/MassCrowd/Public/ZoneGraphCrowdLaneAnnotations.h
@@ -41,10 +41,10 @@ protected:
 	virtual void HandleEvents(TConstArrayView<const UScriptStruct*> AllEventStructs, const FInstancedStructStream& Events) override;
 	virtual void TickAnnotation(const float DeltaTime, FZoneGraphAnnotationTagContainer& AnnotationTagContainer) override;
 
-#if !UE_BUILD_SHIPPING && !UE_BUILD_TEST
+#if UE_ENABLE_DEBUG_DRAWING
 	virtual void DebugDraw(FZoneGraphAnnotationSceneProxy* DebugProxy) override;
 	virtual void DebugDrawCanvas(UCanvas* Canvas, APlayerController*) override;
-#endif
+#endif // UE_ENABLE_DEBUG_DRAWING
 
 	/** Annotation Tag to mark a closed lane. */
 	UPROPERTY(EditAnywhere, Category = CrowdLane)
diff --git a/Engine/Plugins/Runtime/MassGameplay/Source/MassSmartObjects/Private/SmartObjectZoneAnnotations.cpp b/Engine/Plugins/Runtime/MassGameplay/Source/MassSmartObjects/Private/SmartObjectZoneAnnotations.cpp
--- a/Engine/Plugins/Runtime/MassGameplay/Source/MassSmartObjects/Private/SmartObjectZoneAnnotations.cpp
+++ b/Engine/Plugins/Runtime/MassGameplay/Source/MassSmartObjects/Private/SmartObjectZoneAnnotations.cpp
@@ -2,7 +2,6 @@
 
 #include "SmartObjectZoneAnnotations.h"
 #include "MassSmartObjectSettings.h"
-//#include "SmartObjectCollection.h"
 #include "SmartObjectComponent.h"
 #include "SmartObjectSubsystem.h"
 #include "ZoneGraphAnnotationSubsystem.h"
@@ -172,12 +171,12 @@ void USmartObjectZoneAnnotations::TickAnnotation(const float DeltaTime, FZoneGra
 		Data.bInitialTaggingCompleted = true;
 	}
 
-#if !UE_BUILD_SHIPPING && !UE_BUILD_TEST
+#if UE_ENABLE_DEBUG_DRAWING
 	MarkRenderStateDirty();
-#endif
+#endif // UE_ENABLE_DEBUG_DRAWING
 }
 
-#if !UE_BUILD_SHIPPING && !UE_BUILD_TEST
+#if UE_ENABLE_DEBUG_DRAWING
 void USmartObjectZoneAnnotations::DebugDraw(FZoneGraphAnnotationSceneProxy* DebugProxy)
 {
 	UZoneGraphSubsystem* ZoneGraph = UWorld::GetSubsystem<UZoneGraphSubsystem>(GetWorld());
@@ -220,7 +219,7 @@ void USmartObjectZoneAnnotations::DebugDraw(FZoneGraphAnnotationSceneProxy* Debu
 		}
 	}
 }
-#endif // !UE_BUILD_SHIPPING && !UE_BUILD_TEST
+#endif // UE_ENABLE_DEBUG_DRAWING
 
 #if WITH_EDITORONLY_DATA
 void USmartObjectZoneAnnotations::Serialize(FArchive& Ar)
diff --git a/Engine/Plugins/Runtime/MassGameplay/Source/MassSmartObjects/Public/SmartObjectZoneAnnotations.h b/Engine/Plugins/Runtime/MassGameplay/Source/MassSmartObjects/Public/SmartObjectZoneAnnotations.h
--- a/Engine/Plugins/Runtime/MassGameplay/Source/MassSmartObjects/Public/SmartObjectZoneAnnotations.h
+++ b/Engine/Plugins/Runtime/MassGameplay/Source/MassSmartObjects/Public/SmartObjectZoneAnnotations.h
@@ -105,9 +105,9 @@ protected:
 	virtual void PostZoneGraphDataAdded(const AZoneGraphData& ZoneGraphData) override;
 	virtual void PreZoneGraphDataRemoved(const AZoneGraphData& ZoneGraphData) override;
 
-#if !UE_BUILD_SHIPPING && !UE_BUILD_TEST
+#if UE_ENABLE_DEBUG_DRAWING
 	virtual void DebugDraw(FZoneGraphAnnotationSceneProxy* DebugProxy) override;
-#endif
+#endif // UE_ENABLE_DEBUG_DRAWING
 
 	/** Filter specifying which lanes the behavior is applied to. */
 	UPROPERTY(EditAnywhere, Category = SmartObject)
diff --git a/Engine/Plugins/Runtime/ZoneGraphAnnotations/Source/ZoneGraphAnnotations/Private/Annotations/ZoneGraphDisturbanceAnnotation.cpp b/Engine/Plugins/Runtime/ZoneGraphAnnotations/Source/ZoneGraphAnnotations/Private/Annotations/ZoneGraphDisturbanceAnnotation.cpp
--- a/Engine/Plugins/Runtime/ZoneGraphAnnotations/Source/ZoneGraphAnnotations/Private/Annotations/ZoneGraphDisturbanceAnnotation.cpp
+++ b/Engine/Plugins/Runtime/ZoneGraphAnnotations/Source/ZoneGraphAnnotations/Private/Annotations/ZoneGraphDisturbanceAnnotation.cpp
@@ -8,9 +8,9 @@
 #include "ZoneGraphRenderingUtilities.h"
 #include "ZoneGraphQuery.h"
 
-#if !UE_BUILD_SHIPPING && !UE_BUILD_TEST
+#if UE_ENABLE_DEBUG_DRAWING
 #include "Engine/Canvas.h"
-#endif // !UE_BUILD_SHIPPING && !UE_BUILD_TEST
+#endif // UE_ENABLE_DEBUG_DRAWING
 
 #ifdef _MSC_VER
 #pragma warning(push)
@@ -71,7 +71,7 @@ void UZoneGraphDisturbanceAnnotation::TickAnnotation(const float DeltaTime, FZon
 		UpdateAnnotationTags(AnnotationTagContainer);
 	}
 
-#if !UE_BUILD_SHIPPING && !UE_BUILD_TEST
+#if UE_ENABLE_DEBUG_DRAWING
 	if (bEnableDebugDrawing)
 	{
 		FVector ViewLocation = FVector::ZeroVector;
@@ -85,7 +85,7 @@ void UZoneGraphDisturbanceAnnotation::TickAnnotation(const float DeltaTime, FZon
 			MarkRenderStateDirty();
 		}
 	}
-#endif
+#endif // UE_ENABLE_DEBUG_DRAWING
 
 }
 
@@ -656,7 +656,7 @@ FZoneGraphTagMask UZoneGraphDisturbanceAnnotation::GetAnnotationTags() const
 	return AnnotationTags;
 }
 
-#if !UE_BUILD_SHIPPING && !UE_BUILD_TEST
+#if UE_ENABLE_DEBUG_DRAWING
 void UZoneGraphDisturbanceAnnotation::DebugDraw(FZoneGraphAnnotationSceneProxy* DebugProxy)
 {
 	const UZoneGraphSubsystem* ZoneGraph = UWorld::GetSubsystem<UZoneGraphSubsystem>(GetWorld());
@@ -851,13 +851,13 @@ void UZoneGraphDisturbanceAnnotation::DebugDrawCanvas(UCanvas* Canvas, APlayerCo
 
 	Canvas->SetDrawColor(OldDrawColor);
 }
-#endif
+#endif // UE_ENABLE_DEBUG_DRAWING
 
 
 //////////////////////////////////////////////////////////////////////////
 // UZoneGraphDisturbanceAnnotationTest
 
-#if !UE_BUILD_SHIPPING && !UE_BUILD_TEST
+#if UE_ENABLE_DEBUG_DRAWING
 FBox UZoneGraphDisturbanceAnnotationTest::CalcBounds(const FTransform& LocalToWorld) const
 {
 	const FVector Size(DangerRadius);
@@ -876,7 +876,7 @@ void UZoneGraphDisturbanceAnnotationTest::DebugDraw(FDebugRenderSceneProxy* Debu
 	DebugProxy->Boxes.Emplace(FBox(Position - CenterSize, Position + CenterSize), FColor::Red);
 	DebugProxy->Boxes.Emplace(FBox(Position - DangerSize, Position + DangerSize), FColor::Red);
 }
-#endif
+#endif // UE_ENABLE_DEBUG_DRAWING
 
 void UZoneGraphDisturbanceAnnotationTest::Trigger()
 {
diff --git a/Engine/Plugins/Runtime/ZoneGraphAnnotations/Source/ZoneGraphAnnotations/Public/Annotations/ZoneGraphDisturbanceAnnotation.h b/Engine/Plugins/Runtime/ZoneGraphAnnotations/Source/ZoneGraphAnnotations/Public/Annotations/ZoneGraphDisturbanceAnnotation.h
--- a/Engine/Plugins/Runtime/ZoneGraphAnnotations/Source/ZoneGraphAnnotations/Public/Annotations/ZoneGraphDisturbanceAnnotation.h
+++ b/Engine/Plugins/Runtime/ZoneGraphAnnotations/Source/ZoneGraphAnnotations/Public/Annotations/ZoneGraphDisturbanceAnnotation.h
@@ -179,12 +179,12 @@ protected:
 	void UpdateDangerLanes();
 	void UpdateAnnotationTags(FZoneGraphAnnotationTagContainer& AnnotationTagContainer);
 
-#if !UE_BUILD_SHIPPING && !UE_BUILD_TEST
+#if UE_ENABLE_DEBUG_DRAWING
 	virtual void DebugDraw(FZoneGraphAnnotationSceneProxy* DebugProxy) override;
 	virtual void DebugDrawCanvas(UCanvas* Canvas, APlayerController*) override;
 
 	FVector LastDebugDrawLocation = FVector::ZeroVector; 
-#endif
+#endif // UE_ENABLE_DEBUG_DRAWING
 
 	void CalculateEscapeGraph(FZoneGraphDataEscapeGraph& EscapeGraph);
 
@@ -238,11 +238,11 @@ class ZONEGRAPHANNOTATIONS_API UZoneGraphDisturbanceAnnotationTest : public UZon
 
 public:
 	
-#if !UE_BUILD_SHIPPING && !UE_BUILD_TEST
+#if UE_ENABLE_DEBUG_DRAWING
 	virtual FBox CalcBounds(const FTransform& LocalToWorld) const override;
 	virtual void DebugDraw(FDebugRenderSceneProxy* DebugProxy) override;
 
-#endif
+#endif // UE_ENABLE_DEBUG_DRAWING
 
 	virtual void Trigger() override;
 
