StructUtils: Remove ChunkedStructBuffer and InstancedStructStream experiments.

diff --git a/Engine/Plugins/AI/MassCrowd/Source/MassCrowd/Private/ZoneGraphCrowdLaneAnnotations.cpp b/Engine/Plugins/AI/MassCrowd/Source/MassCrowd/Private/ZoneGraphCrowdLaneAnnotations.cpp
--- a/Engine/Plugins/AI/MassCrowd/Source/MassCrowd/Private/ZoneGraphCrowdLaneAnnotations.cpp
+++ b/Engine/Plugins/AI/MassCrowd/Source/MassCrowd/Private/ZoneGraphCrowdLaneAnnotations.cpp
@@ -32,15 +32,16 @@ FZoneGraphTagMask UZoneGraphCrowdLaneAnnotations::GetAnnotationTags() const
 	return AllTags;
 }
 
-void UZoneGraphCrowdLaneAnnotations::HandleEvents(TConstArrayView<const UScriptStruct*> AllEventStructs, const FInstancedStructStream& Events)
+void UZoneGraphCrowdLaneAnnotations::HandleEvents(const FInstancedStructContainer& Events)
 {
-	Events.ForEach([this](const FConstStructView View)
+	for (int32 Index = 0; Index < Events.Num(); Index++)
 	{
-		if (const FZoneGraphCrowdLaneStateChangeEvent* const StateChangeEvent = View.GetPtr<FZoneGraphCrowdLaneStateChangeEvent>())
+		FConstStructView Event = Events[Index];
+		if (const FZoneGraphCrowdLaneStateChangeEvent* const StateChangeEvent = Event.GetPtr<FZoneGraphCrowdLaneStateChangeEvent>())
 		{
 			StateChangeEvents.Add(*StateChangeEvent);
 		}
-	});
+	}
 }
 
 void UZoneGraphCrowdLaneAnnotations::TickAnnotation(const float DeltaTime, FZoneGraphAnnotationTagContainer& AnnotationTagContainer)
diff --git a/Engine/Plugins/AI/MassCrowd/Source/MassCrowd/Public/ZoneGraphCrowdLaneAnnotations.h b/Engine/Plugins/AI/MassCrowd/Source/MassCrowd/Public/ZoneGraphCrowdLaneAnnotations.h
--- a/Engine/Plugins/AI/MassCrowd/Source/MassCrowd/Public/ZoneGraphCrowdLaneAnnotations.h
+++ b/Engine/Plugins/AI/MassCrowd/Source/MassCrowd/Public/ZoneGraphCrowdLaneAnnotations.h
@@ -38,7 +38,7 @@ class MASSCROWD_API UZoneGraphCrowdLaneAnnotations : public UZoneGraphAnnotation
 protected:
 	virtual void PostSubsystemsInitialized() override;
 	virtual FZoneGraphTagMask GetAnnotationTags() const override;
-	virtual void HandleEvents(TConstArrayView<const UScriptStruct*> AllEventStructs, const FInstancedStructStream& Events) override;
+	virtual void HandleEvents(const FInstancedStructContainer& Events) override;
 	virtual void TickAnnotation(const float DeltaTime, FZoneGraphAnnotationTagContainer& AnnotationTagContainer) override;
 
 #if UE_ENABLE_DEBUG_DRAWING
diff --git a/Engine/Plugins/Experimental/StructUtils/Source/StructUtils/Private/ChunkedStructBuffer.cpp b/Engine/Plugins/Experimental/StructUtils/Source/StructUtils/Private/ChunkedStructBuffer.cpp
deleted file mode 100644
diff --git a/Engine/Plugins/Experimental/StructUtils/Source/StructUtils/Private/InstancedStructStream.cpp b/Engine/Plugins/Experimental/StructUtils/Source/StructUtils/Private/InstancedStructStream.cpp
deleted file mode 100644
diff --git a/Engine/Plugins/Experimental/StructUtils/Source/StructUtils/Public/ChunkedStructBuffer.h b/Engine/Plugins/Experimental/StructUtils/Source/StructUtils/Public/ChunkedStructBuffer.h
deleted file mode 100644
diff --git a/Engine/Plugins/Experimental/StructUtils/Source/StructUtils/Public/InstancedStructStream.h b/Engine/Plugins/Experimental/StructUtils/Source/StructUtils/Public/InstancedStructStream.h
deleted file mode 100644
diff --git a/Engine/Plugins/Experimental/StructUtils/Source/StructUtilsTestSuite/Private/ChunkedStructBufferTest.cpp b/Engine/Plugins/Experimental/StructUtils/Source/StructUtilsTestSuite/Private/ChunkedStructBufferTest.cpp
deleted file mode 100644
diff --git a/Engine/Plugins/Experimental/StructUtils/Source/StructUtilsTestSuite/Private/InstancedStructStreamTest.cpp b/Engine/Plugins/Experimental/StructUtils/Source/StructUtilsTestSuite/Private/InstancedStructStreamTest.cpp
deleted file mode 100644
diff --git a/Engine/Plugins/Runtime/MassEntity/Source/MassEntity/Public/MassCommandBuffer.h b/Engine/Plugins/Runtime/MassEntity/Source/MassEntity/Public/MassCommandBuffer.h
--- a/Engine/Plugins/Runtime/MassEntity/Source/MassEntity/Public/MassCommandBuffer.h
+++ b/Engine/Plugins/Runtime/MassEntity/Source/MassEntity/Public/MassCommandBuffer.h
@@ -4,7 +4,6 @@
 
 #include "MassEntityTypes.h"
 #include "MassEntityManager.h"
-#include "InstancedStructStream.h"
 #include "Misc/MTAccessDetector.h"
 #include "MassEntityUtils.h"
 #include "MassCommands.h"
diff --git a/Engine/Plugins/Runtime/ZoneGraphAnnotations/Source/ZoneGraphAnnotations/Private/Annotations/ZoneGraphDisturbanceAnnotation.cpp b/Engine/Plugins/Runtime/ZoneGraphAnnotations/Source/ZoneGraphAnnotations/Private/Annotations/ZoneGraphDisturbanceAnnotation.cpp
--- a/Engine/Plugins/Runtime/ZoneGraphAnnotations/Source/ZoneGraphAnnotations/Private/Annotations/ZoneGraphDisturbanceAnnotation.cpp
+++ b/Engine/Plugins/Runtime/ZoneGraphAnnotations/Source/ZoneGraphAnnotations/Private/Annotations/ZoneGraphDisturbanceAnnotation.cpp
@@ -570,42 +570,46 @@ void UZoneGraphDisturbanceAnnotation::UpdateAnnotationTags(FZoneGraphAnnotationT
 	}
 }
 
-void UZoneGraphDisturbanceAnnotation::HandleEvents(TConstArrayView<const UScriptStruct*> AllEventStructs, const FInstancedStructStream& Events)
+void UZoneGraphDisturbanceAnnotation::HandleEvents(const FInstancedStructContainer& Events)
 {
-	Events.ForEach<FZoneGraphDisturbanceArea>([this](const FZoneGraphDisturbanceArea& Danger)
+	for (int32 Index = 0; Index < Events.Num(); Index++)
 	{
-		FZoneGraphDisturbanceArea* ExistingDanger = Dangers.FindByPredicate([Danger](const FZoneGraphDisturbanceArea& Area) { return Area.InstigatorID == Danger.InstigatorID; });
-		if (ExistingDanger)
-		{
-			*ExistingDanger = Danger;
-		}
-		else
-		{
-			Dangers.Add(Danger);
-		}
-		bDisturbancesChanged = true;
-	});
+		FConstStructView Event = Events[Index];
 
-	Events.ForEach<FZoneGraphObstacleDisturbanceArea>([this](const FZoneGraphObstacleDisturbanceArea& Obstacle)
-	{
-		if (Obstacle.Action == EZoneGraphObstacleDisturbanceAreaAction::Add)
+		if (const FZoneGraphDisturbanceArea* Danger = Event.GetPtr<FZoneGraphDisturbanceArea>())
 		{
-			FZoneGraphObstacleDisturbanceArea* ExistingObstacle = Obstacles.FindByPredicate([Obstacle](const FZoneGraphObstacleDisturbanceArea& Area) { return Area.ObstacleID == Obstacle.ObstacleID; });
-			if (ExistingObstacle)
+			FZoneGraphDisturbanceArea* ExistingDanger = Dangers.FindByPredicate([Danger](const FZoneGraphDisturbanceArea& Area) { return Area.InstigatorID == Danger->InstigatorID; });
+			if (ExistingDanger)
 			{
-				*ExistingObstacle = Obstacle;
+				*ExistingDanger = *Danger;
 			}
 			else
 			{
-				Obstacles.Add(Obstacle);
+				Dangers.Add(*Danger);
 			}
+			bDisturbancesChanged = true;
 		}
-		else if (Obstacle.Action == EZoneGraphObstacleDisturbanceAreaAction::Remove)
+		else if (const FZoneGraphObstacleDisturbanceArea* Obstacle = Event.GetPtr<FZoneGraphObstacleDisturbanceArea>())
 		{
-			Obstacles.Remove(Obstacle);
+			if (Obstacle->Action == EZoneGraphObstacleDisturbanceAreaAction::Add)
+			{
+				FZoneGraphObstacleDisturbanceArea* ExistingObstacle = Obstacles.FindByPredicate([Obstacle](const FZoneGraphObstacleDisturbanceArea& Area) { return Area.ObstacleID == Obstacle->ObstacleID; });
+				if (ExistingObstacle)
+				{
+					*ExistingObstacle = *Obstacle;
+				}
+				else
+				{
+					Obstacles.Add(*Obstacle);
+				}
+			}
+			else if (Obstacle->Action == EZoneGraphObstacleDisturbanceAreaAction::Remove)
+			{
+				Obstacles.Remove(*Obstacle);
+			}
+			bDisturbancesChanged = true;
 		}
-		bDisturbancesChanged = true;
-	});
+	}
 }
 
 void UZoneGraphDisturbanceAnnotation::PostZoneGraphDataAdded(const AZoneGraphData& ZoneGraphData)
diff --git a/Engine/Plugins/Runtime/ZoneGraphAnnotations/Source/ZoneGraphAnnotations/Private/ZoneGraphAnnotationSubsystem.cpp b/Engine/Plugins/Runtime/ZoneGraphAnnotations/Source/ZoneGraphAnnotations/Private/ZoneGraphAnnotationSubsystem.cpp
--- a/Engine/Plugins/Runtime/ZoneGraphAnnotations/Source/ZoneGraphAnnotations/Private/ZoneGraphAnnotationSubsystem.cpp
+++ b/Engine/Plugins/Runtime/ZoneGraphAnnotations/Source/ZoneGraphAnnotations/Private/ZoneGraphAnnotationSubsystem.cpp
@@ -16,6 +16,10 @@ void UZoneGraphAnnotationSubsystem::Initialize(FSubsystemCollectionBase& Collect
 {
 	Super::Initialize(Collection);
 
+	constexpr int32 InitialEventBufferSize = 1024;
+	Events[0].ReserveBytes(InitialEventBufferSize);
+	Events[1].ReserveBytes(InitialEventBufferSize);
+	
 	const UZoneGraphSubsystem* ZoneGraph = Collection.InitializeDependency<UZoneGraphSubsystem>();
 
 	// Register existing data.
@@ -165,18 +169,14 @@ void UZoneGraphAnnotationSubsystem::Tick(float DeltaTime)
 		// Switch exposed current stream, so that event handling can add more events without disturbing the processing.
 		const int32 PrevEventStream = CurrentEventStream;
 		CurrentEventStream ^= 1;
-		FInstancedStructStream& CurrentEvents = Events[PrevEventStream];
+		FInstancedStructContainer& CurrentEvents = Events[PrevEventStream];
 		
-		// Mark which Annotations have received requests
-		TArray<const UScriptStruct*> AllEventStructs;
-		CurrentEvents.GetScriptStructs(AllEventStructs);
-
 		// Process requests
 		for (const FRegisteredZoneGraphAnnotation& Registered : RegisteredComponents)
 		{
 			if (Registered.AnnotationComponent != nullptr)
 			{
-				Registered.AnnotationComponent->HandleEvents(AllEventStructs, CurrentEvents);
+				Registered.AnnotationComponent->HandleEvents(CurrentEvents);
 			}
 		}
 
diff --git a/Engine/Plugins/Runtime/ZoneGraphAnnotations/Source/ZoneGraphAnnotations/Public/Annotations/ZoneGraphDisturbanceAnnotation.h b/Engine/Plugins/Runtime/ZoneGraphAnnotations/Source/ZoneGraphAnnotations/Public/Annotations/ZoneGraphDisturbanceAnnotation.h
--- a/Engine/Plugins/Runtime/ZoneGraphAnnotations/Source/ZoneGraphAnnotations/Public/Annotations/ZoneGraphDisturbanceAnnotation.h
+++ b/Engine/Plugins/Runtime/ZoneGraphAnnotations/Source/ZoneGraphAnnotations/Public/Annotations/ZoneGraphDisturbanceAnnotation.h
@@ -171,7 +171,7 @@ protected:
 	virtual FZoneGraphTagMask GetAnnotationTags() const override;
 
 	virtual void TickAnnotation(const float DeltaTime, FZoneGraphAnnotationTagContainer& AnnotationTagContainer) override;
-	virtual void HandleEvents(TConstArrayView<const UScriptStruct*> AllEventStructs, const FInstancedStructStream& Events) override;
+	virtual void HandleEvents(const FInstancedStructContainer& Events) override;
 
 	virtual void PostZoneGraphDataAdded(const AZoneGraphData& ZoneGraphData) override;
 	virtual void PreZoneGraphDataRemoved(const AZoneGraphData& ZoneGraphData) override;
diff --git a/Engine/Plugins/Runtime/ZoneGraphAnnotations/Source/ZoneGraphAnnotations/Public/ZoneGraphAnnotationComponent.h b/Engine/Plugins/Runtime/ZoneGraphAnnotations/Source/ZoneGraphAnnotations/Public/ZoneGraphAnnotationComponent.h
--- a/Engine/Plugins/Runtime/ZoneGraphAnnotations/Source/ZoneGraphAnnotations/Public/ZoneGraphAnnotationComponent.h
+++ b/Engine/Plugins/Runtime/ZoneGraphAnnotations/Source/ZoneGraphAnnotations/Public/ZoneGraphAnnotationComponent.h
@@ -12,7 +12,7 @@ class UZoneGraphAnnotationSubsystem;
 class UZoneGraphAnnotationComponent;
 class UCanvas;
 class AZoneGraphData;
-struct FInstancedStructStream;
+struct FInstancedStructContainer;
 struct FZoneGraphAnnotationTagLookup;
 struct FZoneGraphAnnotationTagContainer;
 
@@ -47,7 +47,7 @@ public:
 	virtual void TickAnnotation(const float DeltaTime, FZoneGraphAnnotationTagContainer& AnnotationTagContainer) {}
 
 	/** Called when new events are ready to be processed */
-	virtual void HandleEvents(TConstArrayView<const UScriptStruct*> AllEventStructs, const FInstancedStructStream& Events) {}
+	virtual void HandleEvents(const FInstancedStructContainer& Events) {}
 
 	/** @return Tags applied by the Annotation, used to lookup Annotations from tags. */
 	virtual FZoneGraphTagMask GetAnnotationTags() const { return FZoneGraphTagMask::None; }
diff --git a/Engine/Plugins/Runtime/ZoneGraphAnnotations/Source/ZoneGraphAnnotations/Public/ZoneGraphAnnotationSubsystem.h b/Engine/Plugins/Runtime/ZoneGraphAnnotations/Source/ZoneGraphAnnotations/Public/ZoneGraphAnnotationSubsystem.h
--- a/Engine/Plugins/Runtime/ZoneGraphAnnotations/Source/ZoneGraphAnnotations/Public/ZoneGraphAnnotationSubsystem.h
+++ b/Engine/Plugins/Runtime/ZoneGraphAnnotations/Source/ZoneGraphAnnotations/Public/ZoneGraphAnnotationSubsystem.h
@@ -2,7 +2,7 @@
 
 #pragma once
 
-#include "InstancedStructStream.h"
+#include "InstancedStructContainer.h"
 #include "ZoneGraphTypes.h"
 #include "ZoneGraphAnnotationTypes.h"
 #include "Misc/MTAccessDetector.h"
@@ -84,7 +84,8 @@ public:
 	{
 		// This method is thread safe.
 		UE_MT_SCOPED_WRITE_ACCESS(EventsDetector);
-		Events[CurrentEventStream].Add(InRequest);
+		FConstStructView Event = FConstStructView::Make(InRequest);
+		Events[CurrentEventStream].Append(MakeArrayView(&Event, 1));
 	}
 
 	/** @return bitmask of Annotation tags at given lane */
@@ -125,7 +126,7 @@ protected:
 
 	/** Stream of events to be processed, double buffered. */
 	UPROPERTY(Transient)
-	FInstancedStructStream Events[2];
+	FInstancedStructContainer Events[2];
 	UE_MT_DECLARE_RW_ACCESS_DETECTOR(EventsDetector);
 
 	/** Index of the current event stream. */
