[ZoneGraph] made annotation events collection and handling thread safe

diff --git a/Engine/Plugins/Runtime/ZoneGraphAnnotations/Source/ZoneGraphAnnotations/Private/ZoneGraphAnnotationSubsystem.cpp b/Engine/Plugins/Runtime/ZoneGraphAnnotations/Source/ZoneGraphAnnotations/Private/ZoneGraphAnnotationSubsystem.cpp
--- a/Engine/Plugins/Runtime/ZoneGraphAnnotations/Source/ZoneGraphAnnotations/Private/ZoneGraphAnnotationSubsystem.cpp
+++ b/Engine/Plugins/Runtime/ZoneGraphAnnotations/Source/ZoneGraphAnnotations/Private/ZoneGraphAnnotationSubsystem.cpp
@@ -164,13 +164,19 @@ void UZoneGraphAnnotationSubsystem::Tick(float DeltaTime)
 {
 	Super::Tick(DeltaTime);
 
-	if (Events[CurrentEventStream].Num() > 0)
+	// note that it's safe to store the value of CurrentEventStream since the code below is the only place that changes it.
+	const int32 PrevEventStream = CurrentEventStream;
 	{
-		// Switch exposed current stream, so that event handling can add more events without disturbing the processing.
-		const int32 PrevEventStream = CurrentEventStream;
+		// do a thread-safe buffer switch first, then we can do whatever we want to the events
+		UE::TScopeLock ScopeLock(EventsLock);
+
+		// switch exposed current stream, so that event handling can add more events without disturbing the processing.
 		CurrentEventStream ^= 1;
-		FInstancedStructContainer& CurrentEvents = Events[PrevEventStream];
-		
+	}
+
+	FInstancedStructContainer& CurrentEvents = Events[PrevEventStream];
+	if (CurrentEvents.Num() > 0)
+	{
 		// Process requests
 		for (const FRegisteredZoneGraphAnnotation& Registered : RegisteredComponents)
 		{
diff --git a/Engine/Plugins/Runtime/ZoneGraphAnnotations/Source/ZoneGraphAnnotations/Public/ZoneGraphAnnotationSubsystem.h b/Engine/Plugins/Runtime/ZoneGraphAnnotations/Source/ZoneGraphAnnotations/Public/ZoneGraphAnnotationSubsystem.h
--- a/Engine/Plugins/Runtime/ZoneGraphAnnotations/Source/ZoneGraphAnnotations/Public/ZoneGraphAnnotationSubsystem.h
+++ b/Engine/Plugins/Runtime/ZoneGraphAnnotations/Source/ZoneGraphAnnotations/Public/ZoneGraphAnnotationSubsystem.h
@@ -6,6 +6,7 @@
 #include "ZoneGraphTypes.h"
 #include "ZoneGraphAnnotationTypes.h"
 #include "Misc/MTAccessDetector.h"
+#include "Misc/SpinLock.h"
 #include "Subsystems/WorldSubsystem.h"
 
 #include "ZoneGraphAnnotationSubsystem.generated.h"
@@ -82,9 +83,10 @@ public:
 	template <typename T>
 	typename TEnableIf<TIsDerivedFrom<T, FZoneGraphAnnotationEventBase>::IsDerived, void>::Type SendEvent(const T& InRequest)
 	{
-		// This method is thread safe.
-		UE_MT_SCOPED_WRITE_ACCESS(EventsDetector);
 		FConstStructView Event = FConstStructView::Make(InRequest);
+
+		// Allow multiple threads to append events (as long as they don't do so while the subsystem gets ticked)
+		UE::TScopeLock ScopeLock(EventsLock);
 		Events[CurrentEventStream].Append(MakeArrayView(&Event, 1));
 	}
 
@@ -128,9 +130,7 @@ protected:
 	UPROPERTY(Transient)
 	FInstancedStructContainer Events[2];
 	UE_MT_DECLARE_RW_ACCESS_DETECTOR(EventsDetector);
-
-	/** Index of the current event stream. */
-	int32 CurrentEventStream = 0;
+	UE::FSpinLock EventsLock;
 
 	/** Lookup table from tag index to Annotation */
 	UPROPERTY(Transient)
@@ -141,4 +141,11 @@ protected:
 
 	FDelegateHandle OnPostZoneGraphDataAddedHandle;
 	FDelegateHandle OnPreZoneGraphDataRemovedHandle;
+
+private:
+	/**
+	 * Index of the current event stream.
+	 * It's marked as "private" to ensure the correctness of the assumptions regarding when it can change. 
+	 */
+	int32 CurrentEventStream = 0;
 };
